Timing Data Investigator tool for digging into timing files generated during compilation.

diff --git a/Engine/Source/Programs/Timing Data Investigator/Timing Data Investigator/App.config b/Engine/Source/Programs/Timing Data Investigator/Timing Data Investigator/App.config
new file mode 100644
--- /dev/null
+++ b/Engine/Source/Programs/Timing Data Investigator/Timing Data Investigator/App.config	
@@ -0,0 +1,6 @@
+<?xml version="1.0" encoding="utf-8"?>
+<configuration>
+    <startup> 
+        <supportedRuntime version="v4.0" sku=".NETFramework,Version=v4.6.2"/>
+    </startup>
+</configuration>
diff --git a/Engine/Source/Programs/Timing Data Investigator/Timing Data Investigator/App.xaml b/Engine/Source/Programs/Timing Data Investigator/Timing Data Investigator/App.xaml
new file mode 100644
--- /dev/null
+++ b/Engine/Source/Programs/Timing Data Investigator/Timing Data Investigator/App.xaml	
@@ -0,0 +1,8 @@
+﻿<Application x:Class="Timing_Data_Investigator.App"
+             xmlns="http://schemas.microsoft.com/winfx/2006/xaml/presentation"
+             xmlns:x="http://schemas.microsoft.com/winfx/2006/xaml"
+             xmlns:local="clr-namespace:Timing_Data_Investigator"
+             StartupUri="MainWindow.xaml">
+    <Application.Resources>
+    </Application.Resources>
+</Application>
diff --git a/Engine/Source/Programs/Timing Data Investigator/Timing Data Investigator/App.xaml.cs b/Engine/Source/Programs/Timing Data Investigator/Timing Data Investigator/App.xaml.cs
new file mode 100644
--- /dev/null
+++ b/Engine/Source/Programs/Timing Data Investigator/Timing Data Investigator/App.xaml.cs	
@@ -0,0 +1,17 @@
+﻿using System;
+using System.Collections.Generic;
+using System.Configuration;
+using System.Data;
+using System.Linq;
+using System.Threading.Tasks;
+using System.Windows;
+
+namespace Timing_Data_Investigator
+{
+    /// <summary>
+    /// Interaction logic for App.xaml
+    /// </summary>
+    public partial class App : Application
+    {
+    }
+}
diff --git a/Engine/Source/Programs/Timing Data Investigator/Timing Data Investigator/Commands/OpenTimingDataCommand.cs b/Engine/Source/Programs/Timing Data Investigator/Timing Data Investigator/Commands/OpenTimingDataCommand.cs
new file mode 100644
--- /dev/null
+++ b/Engine/Source/Programs/Timing Data Investigator/Timing Data Investigator/Commands/OpenTimingDataCommand.cs	
@@ -0,0 +1,30 @@
+using System;
+using System.Windows.Input;
+using Timing_Data_Investigator.Models;
+
+namespace Timing_Data_Investigator.Commands
+{
+    public class OpenTimingDataCommand : ICommand
+    {
+        private TimingDataViewModel ViewModelToOpen;
+
+		public event EventHandler CanExecuteChanged;
+
+        public OpenTimingDataCommand(TimingDataViewModel TimingData)
+        {
+			ViewModelToOpen = TimingData;
+        }
+
+        public Action<TimingDataViewModel> OpenAction { get; set; }
+
+        public bool CanExecute(object parameter)
+        {
+            return true;
+        }
+
+        public void Execute(object parameter)
+        {
+            OpenAction?.Invoke(ViewModelToOpen);
+        }
+    }
+}
diff --git a/Engine/Source/Programs/Timing Data Investigator/Timing Data Investigator/Controls/TimingDataGrid.xaml b/Engine/Source/Programs/Timing Data Investigator/Timing Data Investigator/Controls/TimingDataGrid.xaml
new file mode 100644
--- /dev/null
+++ b/Engine/Source/Programs/Timing Data Investigator/Timing Data Investigator/Controls/TimingDataGrid.xaml	
@@ -0,0 +1,74 @@
+﻿<UserControl x:Class="Timing_Data_Investigator.Controls.TimingDataGrid"
+             xmlns="http://schemas.microsoft.com/winfx/2006/xaml/presentation"
+             xmlns:x="http://schemas.microsoft.com/winfx/2006/xaml"
+             xmlns:mc="http://schemas.openxmlformats.org/markup-compatibility/2006" 
+             xmlns:d="http://schemas.microsoft.com/expression/blend/2008" 
+             xmlns:local="clr-namespace:Timing_Data_Investigator.Controls"
+             xmlns:tdi="clr-namespace:Timing_Data_Investigator"
+             mc:Ignorable="d" >
+    
+    <UserControl.Resources>
+        <GridLength x:Key="TreeToggleWidth">20</GridLength>
+        <tdi:LevelConverter x:Key="LevelSpaceConverter" LevelWidth="{StaticResource TreeToggleWidth}"/>
+        <tdi:BoolToVisibilityConverter x:Key="BoolToVisibilityConverter" />
+        <tdi:MultiBoolToVisibilityConverter x:Key="MultiBoolToVisibilityConverter" />
+		<tdi:IntLessThanConverter x:Key="IntComparisonConverter" />
+
+		<Style TargetType="{x:Type CheckBox}" x:Key="Expander">
+			<Setter Property="Template">
+                <Setter.Value>
+                    <ControlTemplate TargetType="{x:Type CheckBox}">
+                        <StackPanel Orientation="Horizontal">
+							<Canvas>
+								<Polygon x:Name="Expander" Points="3,3 9,8 3,13" Fill="Black" />
+							</Canvas>
+							<ContentPresenter/>
+                        </StackPanel>
+                        <ControlTemplate.Triggers>
+                            <Trigger Property="IsChecked" Value="True">
+                                <Setter TargetName="Expander" Property="Points" Value="3,4 13,4 8,11" />
+                            </Trigger>
+                            <Trigger Property="IsMouseOver" Value="True">
+                                <Setter TargetName="Expander" Property="Stroke" Value="LightGray"/>
+                                <Setter TargetName="Expander" Property="Fill" Value="DarkGray"/>
+                            </Trigger>
+                        </ControlTemplate.Triggers>
+                    </ControlTemplate>
+                </Setter.Value>
+            </Setter>
+        </Style>
+    </UserControl.Resources>
+
+	<DataGrid AutoGenerateColumns="False" Name="Grid" Grid.Row="1" CanUserAddRows="False" AlternationCount="2" AlternatingRowBackground="#EEE"
+              ItemsSource="{Binding}" Sorting="Grid_Sorting" PreviewKeyDown="Grid_KeyDown">
+		<DataGrid.Columns>
+            <DataGridTemplateColumn Header="Name" Width="*" CanUserSort="True" SortMemberPath="Name">
+                <DataGridTemplateColumn.CellTemplate>
+                    <DataTemplate>
+                        <Grid>
+                            <Grid.ColumnDefinitions>
+                                <ColumnDefinition Width="{Binding Level, Converter={StaticResource LevelSpaceConverter}}"/>
+                                <ColumnDefinition Width="{StaticResource TreeToggleWidth}"/>
+                                <ColumnDefinition Width="*"/>
+                            </Grid.ColumnDefinitions>
+
+                            <CheckBox Grid.Column="1" IsChecked="{Binding IsExpanded, Mode=TwoWay, UpdateSourceTrigger=PropertyChanged}" Visibility="{Binding HasChildren, Converter={StaticResource BoolToVisibilityConverter}}" Style="{StaticResource Expander}" />
+                            <Line Grid.Column="1" X1="8" X2="8" Y2="{Binding ActualHeight, RelativeSource={RelativeSource Self}}" Stroke="#888" StrokeDashArray="2 2" StrokeDashOffset="1">
+                                <Line.Visibility>
+                                    <MultiBinding Converter="{StaticResource MultiBoolToVisibilityConverter}" ConverterParameter="VisibleOnFalse">
+                                        <Binding Path="HasChildren" />
+										<Binding Path="Level" Converter="{StaticResource IntComparisonConverter}" ConverterParameter="1" />
+                                    </MultiBinding>
+                                </Line.Visibility>
+                            </Line>
+                            <TextBlock Grid.Column="2" Text="{Binding Name}"/>
+                        </Grid>
+                    </DataTemplate>
+                </DataGridTemplateColumn.CellTemplate>
+            </DataGridTemplateColumn>
+			<DataGridTextColumn Header="Duration" Binding="{Binding InclusiveDuration, StringFormat={}{0:0.########}s}" SortMemberPath="InclusiveDuration" />
+			<DataGridTextColumn Header="Duration (Exclusive)" Binding="{Binding ExclusiveDuration, StringFormat={}{0:0.########}s}" SortMemberPath="ExclusiveDuration" />
+        </DataGrid.Columns>
+    </DataGrid>
+</UserControl>
+ 
\ No newline at end of file
diff --git a/Engine/Source/Programs/Timing Data Investigator/Timing Data Investigator/Controls/TimingDataGrid.xaml.cs b/Engine/Source/Programs/Timing Data Investigator/Timing Data Investigator/Controls/TimingDataGrid.xaml.cs
new file mode 100644
--- /dev/null
+++ b/Engine/Source/Programs/Timing Data Investigator/Timing Data Investigator/Controls/TimingDataGrid.xaml.cs	
@@ -0,0 +1,84 @@
+﻿using System.ComponentModel;
+using System.Windows.Controls;
+using System.Windows.Data;
+using System.Windows.Input;
+
+namespace Timing_Data_Investigator.Controls
+{
+	/// <summary>
+	/// Interaction logic for TimingDataGrid.xaml
+	/// </summary>
+	public partial class TimingDataGrid : UserControl
+    {
+        public TimingDataGrid()
+        {
+            InitializeComponent();
+        }
+
+        private void Grid_Sorting(object sender, DataGridSortingEventArgs e)
+        {
+			CollectionViewSource GridModel = DataContext as CollectionViewSource;
+			switch (e.Column.SortDirection)
+			{
+				case ListSortDirection.Ascending:
+					{
+						e.Column.SortDirection = ListSortDirection.Descending;
+						break;
+					}
+
+				case ListSortDirection.Descending:
+					{
+						e.Column.SortDirection = null;
+						break;
+					}
+
+				default:
+					{
+						e.Column.SortDirection = ListSortDirection.Ascending;
+						break;
+					}
+			}
+
+			GridModel?.SortDescriptions?.Clear();
+			if (e.Column.SortDirection != null)
+			{
+				GridModel.SortDescriptions.Add(new SortDescription(e.Column.SortMemberPath, e.Column.SortDirection.Value));
+			}
+
+            e.Handled = true;
+        }
+
+        private void Grid_KeyDown(object sender, KeyEventArgs e)
+        {
+			TreeGridElement SelectedData = Grid.SelectedItem as TreeGridElement;
+            if (SelectedData == null)
+            {
+                return;
+            }
+
+            if (e.Key == Key.Right && !SelectedData.IsExpanded)
+            {
+                SelectedData.IsExpanded = true;
+                e.Handled = true;
+            }
+            else if (e.Key == Key.Left)
+            {
+                while (SelectedData != null)
+                {
+                    if (SelectedData.IsExpanded)
+                    {
+                        SelectedData.IsExpanded = false;
+                        Grid.SelectedItem = SelectedData;
+                        Grid.ScrollIntoView(SelectedData);
+						DataGridRow Row = (DataGridRow)Grid.ItemContainerGenerator.ContainerFromItem(SelectedData);
+                        Row.MoveFocus(new TraversalRequest(FocusNavigationDirection.Next));
+                        e.Handled = true;
+                        break;
+                    }
+
+                    SelectedData = SelectedData.Parent as TreeGridElement;
+                }
+            }
+        }
+    }
+}
diff --git a/Engine/Source/Programs/Timing Data Investigator/Timing Data Investigator/Controls/TimingDataTab.xaml b/Engine/Source/Programs/Timing Data Investigator/Timing Data Investigator/Controls/TimingDataTab.xaml
new file mode 100644
--- /dev/null
+++ b/Engine/Source/Programs/Timing Data Investigator/Timing Data Investigator/Controls/TimingDataTab.xaml	
@@ -0,0 +1,66 @@
+﻿<UserControl x:Class="Timing_Data_Investigator.Controls.TimingDataTab"
+             xmlns="http://schemas.microsoft.com/winfx/2006/xaml/presentation"
+             xmlns:x="http://schemas.microsoft.com/winfx/2006/xaml"
+             xmlns:mc="http://schemas.openxmlformats.org/markup-compatibility/2006" 
+             xmlns:d="http://schemas.microsoft.com/expression/blend/2008" 
+             xmlns:local="clr-namespace:Timing_Data_Investigator.Controls"
+             mc:Ignorable="d" 
+             d:DesignHeight="450" d:DesignWidth="800">
+
+    <UserControl.Resources>
+        <ContextMenu  x:Key="RowMenu" DataContext="{Binding PlacementTarget.DataContext, RelativeSource={RelativeSource Self}}">
+            <MenuItem Header="{Binding Name, StringFormat=Open {0} in New Tab}" Command="{Binding OpenCommand}" />
+        </ContextMenu>
+        <Style x:Key="FilesRowStyle" TargetType="{x:Type DataGridRow}">
+            <Setter Property="ContextMenu" Value="{StaticResource RowMenu}" />
+        </Style>
+    </UserControl.Resources>
+
+    <TabControl x:Name="SummaryTabs">
+        <TabItem x:Name="FilesTab" Header="Files" Visibility="Collapsed">
+            <local:TimingDataGrid x:Name="FilesGrid" DockPanel.Dock="Top" />
+        </TabItem>
+        <TabItem x:Name="IncludesTab" Header="Includes">
+			<Grid>
+				<Grid.RowDefinitions>
+					<RowDefinition Height="Auto" />
+					<RowDefinition Height="*" />
+				</Grid.RowDefinitions>
+				<StackPanel Grid.Row="0" x:Name="FlattenIncludesRow" Orientation="Horizontal">
+					<TextBlock Text="Flatten Includes:" VerticalAlignment="Center" Margin="4 2" />
+					<CheckBox x:Name="FlattenIncludes" Checked="FlattenIncludes_Checked" Unchecked="FlattenIncludes_Checked" VerticalAlignment="Center" Margin="2"/>
+				</StackPanel>
+
+				<local:TimingDataGrid Grid.Row="1" x:Name="IncludesGrid" DockPanel.Dock="Top" />
+			</Grid>
+		</TabItem>
+        <TabItem x:Name="ClassesTab" Header="Classes">
+            <Grid>
+                <Grid.RowDefinitions>
+                    <RowDefinition Height="Auto" />
+                    <RowDefinition Height="*" />
+                </Grid.RowDefinitions>
+				<StackPanel Grid.Row="0" x:Name="GroupClassesRow" Orientation="Horizontal">
+					<TextBlock Text="Group Template Classes:" VerticalAlignment="Center" Margin="4 2" />
+					<CheckBox x:Name="GroupClassTemplates" Checked="GroupTemplates_Checked" Unchecked="GroupTemplates_Checked" VerticalAlignment="Center" Margin="2" />
+                </StackPanel>
+
+                <local:TimingDataGrid Grid.Row="1" x:Name="ClassesGrid" />
+            </Grid>
+        </TabItem>
+        <TabItem x:Name="FunctionsTab" Header="Functions">
+            <Grid>
+                <Grid.RowDefinitions>
+                    <RowDefinition Height="Auto" />
+                    <RowDefinition Height="*" />
+                </Grid.RowDefinitions>
+                <StackPanel Grid.Row="0" x:Name="GroupFunctionsRow" Orientation="Horizontal">
+					<TextBlock Text="Group Template Classes:" VerticalAlignment="Center" Margin="4 2" />
+					<CheckBox x:Name="GroupFunctionTemplates" Checked="GroupTemplates_Checked" Unchecked="GroupTemplates_Checked" VerticalAlignment="Center" Margin="2" />
+                </StackPanel>
+
+                <local:TimingDataGrid Grid.Row="1" x:Name="FunctionsGrid" />
+            </Grid>
+        </TabItem>
+    </TabControl>
+</UserControl>
diff --git a/Engine/Source/Programs/Timing Data Investigator/Timing Data Investigator/Controls/TimingDataTab.xaml.cs b/Engine/Source/Programs/Timing Data Investigator/Timing Data Investigator/Controls/TimingDataTab.xaml.cs
new file mode 100644
--- /dev/null
+++ b/Engine/Source/Programs/Timing Data Investigator/Timing Data Investigator/Controls/TimingDataTab.xaml.cs	
@@ -0,0 +1,401 @@
+﻿using System;
+using System.Collections.Generic;
+using System.Collections.ObjectModel;
+using System.ComponentModel;
+using System.Linq;
+using System.Text.RegularExpressions;
+using System.Windows;
+using System.Windows.Controls;
+using System.Windows.Data;
+using Timing_Data_Investigator.Models;
+using UnrealBuildTool;
+
+namespace Timing_Data_Investigator.Controls
+{
+	/// <summary>
+	/// Interaction logic for TimingDataTab.xaml
+	/// </summary>
+	public partial class TimingDataTab : UserControl
+    {
+        private TreeGridModel FilesModel;
+        private TreeGridModel IncludesModel;
+		private TreeGridModel FlattenedIncludesModel;
+		private TreeGridModel ClassesModel;
+        private TreeGridModel GroupedClassesModel;
+        private TreeGridModel FunctionsModel;
+        private TreeGridModel GroupedFunctionsModel;
+
+		private CollectionViewSource FilesViewSource = new CollectionViewSource();
+		private CollectionViewSource IncludesViewSource = new CollectionViewSource();
+		private CollectionViewSource ClassesViewSource = new CollectionViewSource();
+		private CollectionViewSource FunctionsViewSource = new CollectionViewSource();
+
+		private Dictionary<TimingDataViewModel, TabState> TabStates = new Dictionary<TimingDataViewModel, TabState>();
+
+        private class TabState
+        {
+            public Dictionary<TabItem, DataGridColumn> SortColumns { get; } = new Dictionary<TabItem, DataGridColumn>();
+            public Dictionary<TabItem, ListSortDirection?> SortDirections { get; } = new Dictionary<TabItem, ListSortDirection?>();
+
+			public bool? FlattenIncludes { get; set; }
+			public bool? GroupClasses { get; set; }
+            public bool? GroupFunctions { get; set; }
+
+            public void SetTabState(TabItem Tab, DataGrid TabDataGrid)
+            {
+                DataGridColumn SortedColumn = TabDataGrid.Columns.FirstOrDefault(c => c.SortDirection != null);
+                SortColumns[Tab] = SortedColumn;
+                SortDirections[Tab] = SortedColumn?.SortDirection;
+            }
+        }
+
+        public TimingDataTab()
+        {
+            InitializeComponent();
+            DataContextChanged += TimingDataTab_DataContextChanged;
+            FilesGrid.Grid.RowStyle = FindResource("FilesRowStyle") as Style;
+			FilesGrid.DataContext = FilesViewSource;
+			IncludesGrid.DataContext = IncludesViewSource;
+			ClassesGrid.DataContext = ClassesViewSource;
+			FunctionsGrid.DataContext = FunctionsViewSource;
+		}
+
+        private void TimingDataTab_DataContextChanged(object sender, DependencyPropertyChangedEventArgs e)
+        {
+            if (e.NewValue == null)
+            {
+                FilesModel?.Clear();
+                IncludesModel?.Clear();
+                ClassesModel?.Clear();
+                FunctionsModel?.Clear();
+                GroupedClassesModel?.Clear();
+                GroupedFunctionsModel?.Clear();
+                IncludesGrid.DataContext = null;
+                ClassesGrid.DataContext = null;
+                FunctionsGrid.DataContext = null;
+                return;
+            }
+
+            // Save off current relevant state for the old tab if needed.
+            if (e.OldValue != null)
+            {
+				TimingDataViewModel OldTimingData = e.OldValue as TimingDataViewModel;
+                if (!TabStates.ContainsKey(OldTimingData))
+                {
+                    TabStates.Add(OldTimingData, new TabState());
+                }
+
+                TabState TabState = TabStates[OldTimingData];
+                TabState.SetTabState(FilesTab, FilesGrid.Grid);
+                TabState.SetTabState(IncludesTab, IncludesGrid.Grid);
+                TabState.SetTabState(ClassesTab, ClassesGrid.Grid);
+                TabState.SetTabState(FunctionsTab, FunctionsGrid.Grid);
+				TabState.FlattenIncludes = FlattenIncludes.IsChecked;
+                TabState.GroupClasses = GroupClassTemplates.IsChecked;
+                TabState.GroupFunctions = GroupFunctionTemplates.IsChecked;
+            }
+
+			TimingDataViewModel NewTimingData = e.NewValue as TimingDataViewModel;
+
+            // Determine whether we loaded a single file's data or a summary file.
+            if (NewTimingData.Type == TimingDataType.Aggregate)
+            {
+				UpdateUIForAggregate(NewTimingData);
+			}
+            else
+            {
+				UpdateUIForSingleFile(NewTimingData);
+            }
+
+            SummaryTabs.SelectedItem = FilesModel == null ? IncludesTab : FilesTab;
+			FlattenIncludesRow.Visibility = FlattenedIncludesModel != null ? Visibility.Visible : Visibility.Hidden;
+			GroupClassesRow.Visibility = GroupedClassesModel != null ? Visibility.Visible : Visibility.Hidden;
+			GroupFunctionsRow.Visibility = GroupedFunctionsModel != null ? Visibility.Visible : Visibility.Hidden;
+
+			if (TabStates.TryGetValue(NewTimingData, out TabState State))
+			{
+				if (NewTimingData.Type != TimingDataType.Aggregate)
+				{
+					FlattenIncludes.IsChecked = State.FlattenIncludes;
+					GroupClassTemplates.IsChecked = State.GroupClasses;
+					GroupFunctionTemplates.IsChecked = State.GroupFunctions;
+				}
+
+				RefreshTabState(FilesTab, State);
+				RefreshTabState(IncludesTab, State);
+				RefreshTabState(ClassesTab, State);
+				RefreshTabState(FunctionsTab, State);
+
+				SortModel(null, FilesGrid.Grid);
+				SortModel(null, IncludesGrid.Grid);
+				SortModel(null, ClassesGrid.Grid);
+				SortModel(null, FunctionsGrid.Grid);
+			}
+        }
+
+		private void UpdateUIForAggregate(TimingDataViewModel TimingData)
+		{
+			FilesModel = GenerateTreeGridModel(TimingData.Children[0].Children);
+			IncludesModel = null;
+			ClassesModel = null;
+			FunctionsModel = null;
+			FlattenedIncludesModel = null;
+			GroupedClassesModel = null;
+			GroupedFunctionsModel = null;
+
+			UpdateTabVisibility(FilesModel.FlatModel, FilesTab, FilesGrid);
+			UpdateTabVisibility(TimingData.Children[1].Children, IncludesTab, IncludesGrid);
+			UpdateTabVisibility(TimingData.Children[2].Children, ClassesTab, ClassesGrid);
+			UpdateTabVisibility(TimingData.Children[3].Children, FunctionsTab, FunctionsGrid);
+		}
+
+		private void UpdateUIForSingleFile(TimingDataViewModel TimingData)
+		{
+			FilesModel = null;
+			IncludesModel = GenerateTreeGridModel(TimingData.Children[0].Children);
+			ClassesModel = GenerateTreeGridModel(TimingData.Children[1].Children);
+			FunctionsModel = GenerateTreeGridModel(TimingData.Children[2].Children);
+			FlattenedIncludesModel = GenerateFlattenedModel(IncludesModel);
+			GroupedClassesModel = GenerateGroupedModel(ClassesModel);
+			GroupedFunctionsModel = GenerateGroupedModel(FunctionsModel);
+
+			UpdateTabVisibility(null, FilesTab, FilesGrid);
+			UpdateTabVisibility(FlattenIncludes.IsChecked == true ? FlattenedIncludesModel.FlatModel : IncludesModel.FlatModel, IncludesTab, IncludesGrid);
+			UpdateTabVisibility(GroupClassTemplates.IsChecked == true ? GroupedClassesModel.FlatModel : ClassesModel.FlatModel, ClassesTab, ClassesGrid);
+			UpdateTabVisibility(GroupFunctionTemplates.IsChecked == true ? GroupedFunctionsModel.FlatModel : FunctionsModel.FlatModel, FunctionsTab, FunctionsGrid);
+		}
+
+		private void RefreshTabState(TabItem Tab, TabState State)
+		{
+			if (State.SortColumns[Tab] == null || State.SortDirections[Tab] == null)
+			{
+				return;
+			}
+
+			State.SortColumns[Tab].SortDirection = State.SortDirections[Tab];
+		}
+
+        private void SortModel(TreeGridModel Model, DataGrid Grid)
+        {
+            // Find the column that has a sorting value set, if any.
+            DataGridColumn SortedColumn = Grid.Columns.FirstOrDefault(c => c.SortDirection != null);
+            if (SortedColumn == null)
+            {
+                // No sort applied.
+                return;
+            }
+
+			CollectionViewSource CollectionToSort = (CollectionViewSource)Grid.DataContext;
+			CollectionToSort?.SortDescriptions?.Clear();
+			if (SortedColumn.SortDirection != null)
+			{
+				CollectionToSort.SortDescriptions.Add(new SortDescription(SortedColumn.SortMemberPath, SortedColumn.SortDirection.Value));
+			}
+		}
+
+        private void UpdateTabVisibility(ObservableCollection<TreeGridElement> TabModel, TabItem Tab, TimingDataGrid Grid)
+        {
+            Tab.Visibility = TabModel == null ? Visibility.Collapsed : Visibility.Visible;
+			CollectionViewSource Source = (CollectionViewSource)Grid.Grid.DataContext;
+			Source.Source = TabModel;
+        }
+
+		private TreeGridModel GenerateTreeGridModel(IEnumerable<TreeGridElement> Children)
+		{
+			TreeGridModel Model = new TreeGridModel();
+			foreach (TreeGridElement Child in Children)
+			{
+				Model.Add(Child);
+			}
+
+			return Model;
+		}
+
+		private void FlattenIncludesInternal(Dictionary<string, TimingDataViewModel> FlattenedIncludes, IEnumerable<TimingDataViewModel> ViewModels)
+		{
+			foreach (TimingDataViewModel TimingData in ViewModels)
+			{
+				if (FlattenedIncludes.ContainsKey(TimingData.Name))
+				{
+					FlattenedIncludes[TimingData.Name].ExclusiveDuration += TimingData.ExclusiveDuration;
+				}
+				else
+				{
+					TimingDataViewModel FlattenedInclude = new TimingDataViewModel()
+					{
+						Name = TimingData.Name,
+						Type = TimingData.Type,
+						ExclusiveDuration = TimingData.ExclusiveDuration
+					};
+
+					FlattenedIncludes.Add(FlattenedInclude.Name, FlattenedInclude);
+				}
+
+				FlattenIncludesInternal(FlattenedIncludes, TimingData.Children.Cast<TimingDataViewModel>());
+			}
+		}
+
+		private TreeGridModel GenerateFlattenedModel(TreeGridModel UnflattenedModel)
+		{
+			Dictionary<string, TimingDataViewModel> FlattenedIncludes = new Dictionary<string, TimingDataViewModel>();
+			FlattenIncludesInternal(FlattenedIncludes, UnflattenedModel.Cast<TimingDataViewModel>());
+			TreeGridModel FlattenedModel = new TreeGridModel();
+			foreach (TimingDataViewModel TimingData in FlattenedIncludes.Values)
+			{
+				FlattenedModel.Add(TimingData);
+			}
+
+			return FlattenedModel;
+		}
+
+		private TreeGridModel GenerateGroupedModel(TreeGridModel UngroupedModel)
+        {
+			TreeGridModel GroupedModel = new TreeGridModel();
+			List<TimingDataViewModel> ClonedChildren = new List<TimingDataViewModel>();
+            foreach (TimingDataViewModel Child in UngroupedModel.Cast<TimingDataViewModel>())
+            {
+				TimingDataViewModel ClonedChild = Child.Clone();
+                if (ClonedChild.HasChildren)
+                {
+					IEnumerable<TimingDataViewModel> GroupedChildren = GroupChildren(ClonedChild.Children.Cast<TimingDataViewModel>());
+                    ClonedChild.Children.Clear();
+                    foreach (TimingDataViewModel GroupedChild in GroupedChildren)
+                    {
+                        ClonedChild.Children.Add(GroupedChild);
+                    }
+                }
+                ClonedChildren.Add(ClonedChild);
+            }
+
+			IEnumerable<TimingDataViewModel> GroupedClonedChildren = GroupChildren(ClonedChildren);
+            foreach (TimingDataViewModel GroupedChild in GroupedClonedChildren)
+            {
+                GroupedModel.Add(GroupedChild);
+            }
+
+            return GroupedModel;
+        }
+
+        private IEnumerable<TimingDataViewModel> GroupChildren(IEnumerable<TimingDataViewModel> Children)
+        {
+			List<TimingDataViewModel> GroupedChildren = new List<TimingDataViewModel>();
+			Dictionary<string, List<TimingDataViewModel>> ChildGroups = new Dictionary<string, List<TimingDataViewModel>>();
+            foreach (TimingDataViewModel Child in Children)
+            {
+                if (Child.HasChildren)
+                {
+					IEnumerable<TimingDataViewModel> ChildsGroupedChildren = GroupChildren(Child.Children.Cast<TimingDataViewModel>());
+                    Child.Children.Clear();
+                    foreach (TimingDataViewModel ChildChild in ChildsGroupedChildren)
+                    {
+                        Child.Children.Add(ChildChild);
+                    }
+                }
+
+				// See if this is a templated class. If not, add it as is.
+				Match Match = Regex.Match(Child.Name, @"^([^<]*)(?<Template><.*>)");
+                if (!Match.Success)
+                {
+					// Check to see if we've seen this name before. If so, group them together.
+					string DuplicateGroupName = $"{Child.Name} (Duplicates)";
+                    if (ChildGroups.ContainsKey(DuplicateGroupName))
+                    {
+                        ChildGroups[DuplicateGroupName].Add(Child);
+                    }
+                    else
+                    {
+						TimingDataViewModel FoundChild = GroupedChildren.FirstOrDefault(c => c.Name == Child.Name);
+                        if (FoundChild != null)
+                        {
+                            ChildGroups.Add(DuplicateGroupName, new List<TimingDataViewModel>());
+                            ChildGroups[DuplicateGroupName].Add(FoundChild);
+                            ChildGroups[DuplicateGroupName].Add(Child);
+                            GroupedChildren.Remove(FoundChild);
+                        }
+                        else
+                        {
+                            GroupedChildren.Add(Child);
+                        }
+                    }
+                }
+                else
+                {
+					// Generate group name from template.
+					int TemplateParamCount = Match.Groups["Template"].Value.Count(c => c == ',') + 1;
+					List<string> TemplateParamSig = new List<string>(TemplateParamCount);
+                    for (int i = 0; i < TemplateParamCount; ++i)
+                    {
+                        TemplateParamSig.Add("...");
+                    }
+
+					string GroupName = Child.Name.Replace(Match.Groups["Template"].Value, $"<{string.Join(", ", TemplateParamSig)}>");
+
+                    // See if we have a group for this template already. If not, add it.
+                    if (!ChildGroups.ContainsKey(GroupName))
+                    {
+                        ChildGroups.Add(GroupName, new List<TimingDataViewModel>());
+                    }
+
+                    ChildGroups[GroupName].Add(Child);
+                }
+            }
+
+            // Add grouped children.
+            foreach (KeyValuePair<string, List<TimingDataViewModel>> Group in ChildGroups)
+            {
+                if (Group.Value.Count == 1)
+                {
+                    GroupedChildren.Add(Group.Value.First());
+                    continue;
+                }
+
+				TimingDataViewModel NewViewModel = new TimingDataViewModel()
+                {
+                    Name = Group.Key,
+                    HasChildren = true,
+                };
+
+                foreach (TimingDataViewModel Child in Group.Value)
+                {
+                    NewViewModel.Children.Add(Child);
+                }
+
+                GroupedChildren.Add(NewViewModel);
+            }
+
+            return GroupedChildren;
+        }
+
+        private void UpdateGridModel(CheckBox CheckBox, TimingDataGrid Grid, TreeGridModel UncheckedModel, TreeGridModel CheckedModel)
+        {
+			CollectionViewSource Source = (CollectionViewSource)Grid.DataContext;
+			if (CheckBox.IsChecked == true)
+            {
+				Source.Source = CheckedModel?.FlatModel;
+            }
+            else
+            {
+				Source.Source = UncheckedModel?.FlatModel;
+            }
+        }
+
+        private void FlattenIncludes_Checked(object sender, RoutedEventArgs e)
+        {
+			CheckBox GroupCheckBox = sender as CheckBox;
+            UpdateGridModel(FlattenIncludes, IncludesGrid, IncludesModel, FlattenedIncludesModel);
+        }
+
+		private void GroupTemplates_Checked(object sender, RoutedEventArgs e)
+		{
+			CheckBox GroupCheckBox = sender as CheckBox;
+			if (GroupCheckBox == GroupClassTemplates)
+			{
+				UpdateGridModel(GroupClassTemplates, ClassesGrid, ClassesModel, GroupedClassesModel);
+			}
+			else
+			{
+				UpdateGridModel(GroupFunctionTemplates, FunctionsGrid, FunctionsModel, GroupedFunctionsModel);
+			}
+		}
+	}
+}
diff --git a/Engine/Source/Programs/Timing Data Investigator/Timing Data Investigator/Converters.cs b/Engine/Source/Programs/Timing Data Investigator/Timing Data Investigator/Converters.cs
new file mode 100644
--- /dev/null
+++ b/Engine/Source/Programs/Timing Data Investigator/Timing Data Investigator/Converters.cs	
@@ -0,0 +1,184 @@
+﻿using System;
+using System.Globalization;
+using System.Windows;
+using System.Windows.Data;
+
+namespace Timing_Data_Investigator
+{
+    public enum VisibilityOptions
+    {
+        VisibleOnFalse,
+        VisibleOnTrue,
+        HiddenOnFalse,
+        HiddenOnTrue,
+        CollapsedOnFalse,
+        CollapsedOnTrue
+    }
+
+    public class NullToVisibilityConverter : IValueConverter
+    {
+        public object Convert(object value, Type targetType, object parameter, CultureInfo culture)
+        {
+            var VisibilityOption = parameter == null ? VisibilityOptions.VisibleOnTrue : Enum.Parse(typeof(VisibilityOptions), parameter.ToString());
+            switch (VisibilityOption)
+            {
+                case VisibilityOptions.VisibleOnFalse:
+                case VisibilityOptions.HiddenOnTrue:
+                    {
+                        return value != null ? Visibility.Hidden : Visibility.Visible;
+                    }
+
+                case VisibilityOptions.CollapsedOnTrue:
+                    {
+                        return value != null ? Visibility.Collapsed : Visibility.Visible;
+                    }
+
+                case VisibilityOptions.CollapsedOnFalse:
+                    {
+                        return value != null ? Visibility.Visible : Visibility.Collapsed;
+                    }
+            }
+
+            return value != null ? Visibility.Visible : Visibility.Hidden;
+        }
+
+        public object ConvertBack(object value, Type targetType, object parameter, CultureInfo culture)
+        {
+            throw new NotImplementedException();
+        }
+
+    }
+    public class BoolToVisibilityConverter : IValueConverter
+    {
+        public object Convert(object value, Type targetType, object parameter, CultureInfo culture)
+        {
+			object VisibilityOption = parameter == null ? VisibilityOptions.VisibleOnTrue : Enum.Parse(typeof(VisibilityOptions), parameter.ToString());
+            switch (VisibilityOption)
+            {
+                case VisibilityOptions.VisibleOnFalse:
+                case VisibilityOptions.HiddenOnTrue:
+                    {
+                        return (bool)value ? Visibility.Hidden: Visibility.Visible;
+                    }
+
+                case VisibilityOptions.CollapsedOnTrue:
+                    {
+                        return (bool)value ? Visibility.Collapsed : Visibility.Visible;
+                    }
+
+                case VisibilityOptions.CollapsedOnFalse:
+                    {
+                        return (bool)value ? Visibility.Visible : Visibility.Collapsed;
+                    }
+            }
+
+            return (bool)value ? Visibility.Visible : Visibility.Hidden;
+        }
+
+        public object ConvertBack(object value, Type targetType, object parameter, CultureInfo culture)
+        {
+            throw new NotImplementedException();
+        }
+    }
+
+    public class MultiBoolToVisibilityConverter : IMultiValueConverter
+    {
+        public object Convert(object[] values, Type targetType, object parameter, CultureInfo culture)
+        {
+			bool aggregateBool = false;
+            foreach (object value in values)
+            {
+                aggregateBool |= (bool)value;
+            }
+
+			VisibilityOptions VisibilityOption = parameter == null ? VisibilityOptions.VisibleOnTrue : (VisibilityOptions)Enum.Parse(typeof(VisibilityOptions), parameter.ToString());
+            switch (VisibilityOption)
+            {
+                case VisibilityOptions.VisibleOnFalse:
+                case VisibilityOptions.HiddenOnTrue:
+                    {
+                        return aggregateBool ? Visibility.Hidden : Visibility.Visible;
+                    }
+
+                case VisibilityOptions.CollapsedOnTrue:
+                    {
+                        return aggregateBool ? Visibility.Collapsed : Visibility.Visible;
+                    }
+
+                case VisibilityOptions.CollapsedOnFalse:
+                    {
+                        return aggregateBool ? Visibility.Visible : Visibility.Collapsed;
+                    }
+            }
+
+            return aggregateBool ? Visibility.Visible : Visibility.Hidden;
+        }
+
+        public object[] ConvertBack(object value, Type[] targetTypes, object parameter, CultureInfo culture)
+        {
+            throw new NotImplementedException();
+        }
+    }
+
+    public class IntGreaterThanConverter : IValueConverter
+    {
+        public object Convert(object value, Type targetType, object parameter, CultureInfo culture)
+        {
+            if (parameter == null)
+            {
+                throw new InvalidOperationException("A ConverterParameter must be set for IntGreaterThanConverter!");
+            }
+
+            if (int.TryParse(parameter.ToString(), out int parameterInt))
+            {
+                return (int)value > parameterInt;
+            }
+
+            throw new InvalidOperationException("ConverterParameter must be a valid int for IntGreaterThanConverter!");
+        }
+
+        public object ConvertBack(object value, Type targetType, object parameter, CultureInfo culture)
+        {
+            throw new NotImplementedException();
+        }
+    }
+
+    public class IntLessThanConverter : IValueConverter
+    {
+        public object Convert(object value, Type targetType, object parameter, CultureInfo culture)
+        {
+            if (parameter == null)
+            {
+                throw new InvalidOperationException("A ConverterParameter must be set for IntLessThanConverter!");
+            }
+
+            if (int.TryParse(parameter.ToString(), out int parameterInt))
+            {
+                return (int)value < parameterInt;
+            }
+
+            throw new InvalidOperationException("ConverterParameter must be a valid int for IntLessThanConverter!");
+        }
+
+        public object ConvertBack(object value, Type targetType, object parameter, CultureInfo culture)
+        {
+            throw new NotImplementedException();
+        }
+    }
+
+	public class LevelConverter : IValueConverter
+    {
+        public GridLength LevelWidth { get; set; }
+
+        public object Convert(object value, Type targetType, object parameter, CultureInfo culture)
+        {
+			// Return the width multiplied by the level
+            return ((int)value * LevelWidth.Value);
+        }
+
+        public object ConvertBack(object value, Type targetType, object parameter, CultureInfo culture)
+        {
+            throw new NotImplementedException();
+        }
+    }
+}
diff --git a/Engine/Source/Programs/Timing Data Investigator/Timing Data Investigator/MainWindow.xaml b/Engine/Source/Programs/Timing Data Investigator/Timing Data Investigator/MainWindow.xaml
new file mode 100644
--- /dev/null
+++ b/Engine/Source/Programs/Timing Data Investigator/Timing Data Investigator/MainWindow.xaml	
@@ -0,0 +1,67 @@
+﻿<Window x:Class="Timing_Data_Investigator.MainWindow"
+        xmlns="http://schemas.microsoft.com/winfx/2006/xaml/presentation"
+        xmlns:x="http://schemas.microsoft.com/winfx/2006/xaml"
+        xmlns:d="http://schemas.microsoft.com/expression/blend/2008"
+        xmlns:mc="http://schemas.openxmlformats.org/markup-compatibility/2006"
+        xmlns:controls="clr-namespace:Timing_Data_Investigator.Controls"
+        xmlns:tdi="clr-namespace:Timing_Data_Investigator"
+        mc:Ignorable="d"
+        Title="Timing Data Investigator" Width="800" Height="600">
+
+    <Window.Resources>
+		<tdi:NullToVisibilityConverter x:Key="NullToVisibilityConverter" />
+		<Style x:Key="TabCloseButton" TargetType="Button">
+			<Setter Property="OverridesDefaultStyle" Value="True"/>
+			<Setter Property="Background" Value="Transparent" />
+			<Setter Property="Template">
+				<Setter.Value>
+					<ControlTemplate TargetType="Button">
+						<Canvas Width="15" Height="15" Margin="2">
+							<Ellipse Width="15" Height="15" Fill="{Binding Background, RelativeSource={RelativeSource Mode=FindAncestor, AncestorType=Button}}" />
+							<Line X1="4" Y1="4" X2="12" Y2="12" Stroke="#222" StrokeThickness="1" />
+							<Line X1="4" Y1="12" X2="12" Y2="4" Stroke="#222" StrokeThickness="1" />
+						</Canvas>
+					</ControlTemplate>
+				</Setter.Value>
+			</Setter>
+			<Style.Triggers>
+				<Trigger Property="IsMouseOver" Value="True">
+					<Setter Property="Background" Value="LightGray" />
+				</Trigger>
+			</Style.Triggers>
+		</Style>
+	</Window.Resources>
+    
+    <Grid>
+        <Grid.RowDefinitions>
+            <RowDefinition Height="*" />
+            <RowDefinition Height="Auto" />
+        </Grid.RowDefinitions>
+        <DockPanel Grid.Row="0">
+            <Menu DockPanel.Dock="Top">
+                <MenuItem Header="_File">
+                    <MenuItem Header="_Open" Click="OpenFile_Click" />
+                    <Separator />
+                    <MenuItem Header="E_xit" Click="Exit_Click" />
+                </MenuItem>
+            </Menu>
+            <TabControl x:Name="OpenedFiles">
+                <TabControl.ItemTemplate>
+                    <DataTemplate>
+                        <StackPanel Orientation="Horizontal">
+							<TextBlock Text="{Binding ShortName}" VerticalAlignment="Center" />
+							<Button Click="RemoveTab_Click" Style="{StaticResource TabCloseButton}" />
+						</StackPanel>
+                    </DataTemplate>
+                </TabControl.ItemTemplate>
+                <TabControl.ContentTemplate>
+                    <DataTemplate>
+                        <controls:TimingDataTab />
+                    </DataTemplate>
+                </TabControl.ContentTemplate>
+                <TabItem x:Name="NoOpenTabsTab" Header="Open a File to Investigate..." Opacity="0.4"/>
+            </TabControl>
+        </DockPanel>
+		<ProgressBar x:Name="LoadingProgress" Grid.Row="0" Height="16" Width="200" Margin="10 21" VerticalAlignment="Top" HorizontalAlignment="Right" Visibility="Collapsed" IsIndeterminate="True" />
+	</Grid>
+</Window>
diff --git a/Engine/Source/Programs/Timing Data Investigator/Timing Data Investigator/MainWindow.xaml.cs b/Engine/Source/Programs/Timing Data Investigator/Timing Data Investigator/MainWindow.xaml.cs
new file mode 100644
--- /dev/null
+++ b/Engine/Source/Programs/Timing Data Investigator/Timing Data Investigator/MainWindow.xaml.cs	
@@ -0,0 +1,87 @@
+﻿using Microsoft.Win32;
+using System;
+using System.Linq;
+using System.Threading.Tasks;
+using System.Windows;
+using System.Windows.Controls;
+using Timing_Data_Investigator.Commands;
+using Timing_Data_Investigator.Models;
+using Tools.DotNETCommon;
+
+namespace Timing_Data_Investigator
+{
+	/// <summary>
+	/// Interaction logic for MainWindow.xaml
+	/// </summary>
+	public partial class MainWindow : Window
+    {
+        public MainWindow()
+        {
+            InitializeComponent();
+        }
+
+        private void LoadTimingFile(string FilePath)
+        {
+			TimingDataViewModel NewTimingData = TimingDataViewModel.FromBinaryFile(FileReference.FromString(FilePath));
+
+			// If this is an aggregate, hook up the open commands for the file rows.
+			if (NewTimingData.Type == UnrealBuildTool.TimingDataType.Aggregate)
+			{
+				foreach (TimingDataViewModel File in NewTimingData.Children[0].Children.Cast<TimingDataViewModel>())
+				{
+					OpenTimingDataCommand FileOpenCommand = new OpenTimingDataCommand(File);
+					FileOpenCommand.OpenAction =
+						(ViewModel) =>
+						{
+							TimingDataViewModel FileTimingData = NewTimingData.LoadTimingDataFromBinaryBlob(File.Name);
+							Dispatcher.BeginInvoke(new Action(() => { AddTimingDataViewModelToTabs(FileTimingData); }));
+						};
+
+					File.OpenCommand = FileOpenCommand;
+				}
+			}
+
+			AddTimingDataViewModelToTabs(NewTimingData);
+        }
+
+        private void AddTimingDataViewModelToTabs(TimingDataViewModel NewViewModel)
+        {
+			NoOpenTabsTab.Visibility = Visibility.Collapsed;
+			OpenedFiles.Items.Add(NewViewModel);
+			OpenedFiles.SelectedItem = NewViewModel;
+        }
+
+        private void OpenFile_Click(object sender, RoutedEventArgs e)
+        {
+			OpenFileDialog OpenFileDialog = new OpenFileDialog();
+            OpenFileDialog.Filter = "Timing Files (*.timing.bin)|*.timing.bin|All Files (*.*)|*.*";
+            if (OpenFileDialog.ShowDialog() == true)
+            {
+                LoadTimingFile(OpenFileDialog.FileName);
+            }
+        }
+
+        private void Exit_Click(object sender, RoutedEventArgs e)
+        {
+            Close();
+        }
+
+        private void RemoveTab_Click(object sender, RoutedEventArgs e)
+        {
+			Button RemoveTabButton = (Button)e.Source;
+			TimingDataViewModel TabToRemove = (TimingDataViewModel)RemoveTabButton.DataContext;
+            int IndexToSelect = OpenedFiles.Items.IndexOf(TabToRemove) - 1;
+            OpenedFiles.Items.Remove(TabToRemove);
+            if (IndexToSelect == 0 && OpenedFiles.Items.Count == 1)
+            {
+                NoOpenTabsTab.Visibility = Visibility.Visible;
+            }
+            else
+            {
+                IndexToSelect = 1;
+            }
+
+            OpenedFiles.SelectedIndex = IndexToSelect;
+        }
+    }
+}
diff --git a/Engine/Source/Programs/Timing Data Investigator/Timing Data Investigator/Models/TimingDataViewModel.cs b/Engine/Source/Programs/Timing Data Investigator/Timing Data Investigator/Models/TimingDataViewModel.cs
new file mode 100644
--- /dev/null
+++ b/Engine/Source/Programs/Timing Data Investigator/Timing Data Investigator/Models/TimingDataViewModel.cs	
@@ -0,0 +1,156 @@
+﻿using System.Collections.Generic;
+using System.Collections.ObjectModel;
+using System.IO;
+using System.IO.Compression;
+using System.Linq;
+using System.Windows.Input;
+using Tools.DotNETCommon;
+using UnrealBuildTool;
+
+namespace Timing_Data_Investigator.Models
+{
+	public class TimingDataViewModel : TreeGridElement
+	{
+		private struct BinaryBlob
+		{
+			public int Offset { get; set; }
+			public int CompressedSize { get; set; }
+			public int DecompressedSize { get; set; }
+		}
+
+		private Dictionary<string, BinaryBlob> BinaryBlobLookupTable;
+		private byte[] BinaryBlobBytes;
+
+		public TimingDataViewModel()
+		{
+		}
+
+		public TimingDataViewModel(IEnumerable<TreeGridElement> InitialChildren)
+		{
+			Children = new ObservableCollection<TreeGridElement>(InitialChildren);
+		}
+
+		public string Name { get; set; }
+		public TimingDataType Type { get; set; }
+		public double ExclusiveDuration { get; set; }
+		public double InclusiveDuration => ExclusiveDuration + Children.Cast<TimingDataViewModel>().Sum(c => c.InclusiveDuration);
+
+		public string ShortName => !string.IsNullOrWhiteSpace(Name) ? Path.GetFileName(Name) : null;
+
+		public ICommand OpenCommand { get; set; }
+
+		public static TimingDataViewModel FromBinaryFile(FileReference BinaryFile)
+		{
+			using (BinaryReader Reader = new BinaryReader(File.Open(BinaryFile.FullName, FileMode.Open, FileAccess.Read)))
+			{
+				return FromBinaryReader(Reader);
+			}
+		}
+
+		public static TimingDataViewModel FromJsonFile(FileReference JsonFile)
+		{
+			string JsonString = File.ReadAllText(JsonFile.FullName);
+			TimingData DeserializedTimingData = Json.Deserialize<TimingData>(JsonString);
+			return FromTimingData(DeserializedTimingData);
+		}
+
+		public static TimingDataViewModel FromTimingData(TimingData TimingData)
+		{
+			TimingDataViewModel NewViewModel = new TimingDataViewModel()
+			{
+				Name = TimingData.Name,
+				Type = TimingData.Type,
+				ExclusiveDuration = TimingData.ExclusiveDuration,
+				HasChildren = TimingData.Children.Any(),
+			};
+
+			foreach (KeyValuePair<string, TimingData> Child in TimingData.Children)
+			{
+				TimingDataViewModel ChildData = FromTimingData(Child.Value);
+				NewViewModel.Children.Add(ChildData);
+			}
+
+			return NewViewModel;
+		}
+
+		public TimingDataViewModel LoadTimingDataFromBinaryBlob(string BinaryBlobName)
+		{
+			if (!BinaryBlobLookupTable.ContainsKey(BinaryBlobName))
+			{
+				return null;
+			}
+
+			BinaryBlob Blob = BinaryBlobLookupTable[BinaryBlobName];
+			using (MemoryStream BlobStream = new MemoryStream(BinaryBlobBytes, Blob.Offset, Blob.CompressedSize))
+			using (GZipStream DecompressionStream = new GZipStream(BlobStream, CompressionMode.Decompress))
+			{
+				byte[] DecompressedBytes = new byte[Blob.DecompressedSize];
+				using (MemoryStream DecompressedStream = new MemoryStream(DecompressedBytes))
+				{
+					DecompressionStream.CopyTo(DecompressedStream);
+					DecompressedStream.Seek(0, SeekOrigin.Begin);
+					using (BinaryReader Reader = new BinaryReader(DecompressedStream))
+					{
+						return FromBinaryReader(Reader);
+					}
+				}
+			}
+		}
+
+		private static TimingDataViewModel FromBinaryReader(BinaryReader Reader)
+		{
+			// Read in the timing data.
+			TimingData DeserializedTimingData = new TimingData(Reader);
+			TimingDataViewModel ViewModel = FromTimingData(DeserializedTimingData);
+
+			// If this is an aggregate, read in the look up table and the binary blobs, and also de-parent the
+			// include, class, and function aggregate lists.
+			if (ViewModel.Type == TimingDataType.Aggregate)
+			{
+				ViewModel.BinaryBlobLookupTable = new Dictionary<string, BinaryBlob>();
+				int BinaryBlobCount = Reader.ReadInt32();
+				for (int i = 0; i < BinaryBlobCount; ++i)
+				{
+					string BlobName = Reader.ReadString();
+					int BlobOffset = Reader.ReadInt32();
+					int BlobCompressedSize = Reader.ReadInt32();
+					int BlobDecompressedSize = Reader.ReadInt32();
+					ViewModel.BinaryBlobLookupTable.Add(BlobName, new BinaryBlob() { Offset = BlobOffset, CompressedSize = BlobCompressedSize, DecompressedSize = BlobDecompressedSize });
+				}
+
+				// Allocate the memory for the binary blobs then copy them into it.
+				int BinaryBlobLength = (int)(Reader.BaseStream.Length - Reader.BaseStream.Position);
+				ViewModel.BinaryBlobBytes = Reader.ReadBytes(BinaryBlobLength);
+
+				foreach (TreeGridElement SummaryChild in ViewModel.Children.Skip(1))
+				{
+					foreach (TreeGridElement Child in SummaryChild.Children)
+					{
+						Child.Level = 0;
+					}
+				}
+			}
+
+			return ViewModel;
+		}
+
+		public TimingDataViewModel Clone()
+		{
+			TimingDataViewModel NewViewModel = new TimingDataViewModel()
+			{
+				Name = Name,
+				Type = Type,
+				ExclusiveDuration = ExclusiveDuration,
+				HasChildren = Children.Any(),
+			};
+
+			foreach (TimingDataViewModel Child in Children.Cast<TimingDataViewModel>())
+			{
+				TimingDataViewModel ChildData = Child.Clone();
+				NewViewModel.Children.Add(ChildData);
+			}
+
+			return NewViewModel;
+		}
+	}
+}
\ No newline at end of file
diff --git a/Engine/Source/Programs/Timing Data Investigator/Timing Data Investigator/Properties/AssemblyInfo.cs b/Engine/Source/Programs/Timing Data Investigator/Timing Data Investigator/Properties/AssemblyInfo.cs
new file mode 100644
--- /dev/null
+++ b/Engine/Source/Programs/Timing Data Investigator/Timing Data Investigator/Properties/AssemblyInfo.cs	
@@ -0,0 +1,55 @@
+﻿using System.Reflection;
+using System.Resources;
+using System.Runtime.CompilerServices;
+using System.Runtime.InteropServices;
+using System.Windows;
+
+// General Information about an assembly is controlled through the following
+// set of attributes. Change these attribute values to modify the information
+// associated with an assembly.
+[assembly: AssemblyTitle("Timing Data Investigator")]
+[assembly: AssemblyDescription("")]
+[assembly: AssemblyConfiguration("")]
+[assembly: AssemblyCompany("")]
+[assembly: AssemblyProduct("Timing Data Investigator")]
+[assembly: AssemblyCopyright("Copyright ©  2019")]
+[assembly: AssemblyTrademark("")]
+[assembly: AssemblyCulture("")]
+
+// Setting ComVisible to false makes the types in this assembly not visible
+// to COM components.  If you need to access a type in this assembly from
+// COM, set the ComVisible attribute to true on that type.
+[assembly: ComVisible(false)]
+
+//In order to begin building localizable applications, set
+//<UICulture>CultureYouAreCodingWith</UICulture> in your .csproj file
+//inside a <PropertyGroup>.  For example, if you are using US english
+//in your source files, set the <UICulture> to en-US.  Then uncomment
+//the NeutralResourceLanguage attribute below.  Update the "en-US" in
+//the line below to match the UICulture setting in the project file.
+
+//[assembly: NeutralResourcesLanguage("en-US", UltimateResourceFallbackLocation.Satellite)]
+
+
+[assembly: ThemeInfo(
+    ResourceDictionaryLocation.None, //where theme specific resource dictionaries are located
+                                     //(used if a resource is not found in the page,
+                                     // or application resource dictionaries)
+    ResourceDictionaryLocation.SourceAssembly //where the generic resource dictionary is located
+                                              //(used if a resource is not found in the page,
+                                              // app, or any theme specific resource dictionaries)
+)]
+
+
+// Version information for an assembly consists of the following four values:
+//
+//      Major Version
+//      Minor Version
+//      Build Number
+//      Revision
+//
+// You can specify all the values or you can default the Build and Revision Numbers
+// by using the '*' as shown below:
+// [assembly: AssemblyVersion("1.0.*")]
+[assembly: AssemblyVersion("1.0.0.0")]
+[assembly: AssemblyFileVersion("1.0.0.0")]
diff --git a/Engine/Source/Programs/Timing Data Investigator/Timing Data Investigator/Properties/Settings.settings b/Engine/Source/Programs/Timing Data Investigator/Timing Data Investigator/Properties/Settings.settings
new file mode 100644
--- /dev/null
+++ b/Engine/Source/Programs/Timing Data Investigator/Timing Data Investigator/Properties/Settings.settings	
@@ -0,0 +1,7 @@
+﻿<?xml version='1.0' encoding='utf-8'?>
+<SettingsFile xmlns="uri:settings" CurrentProfile="(Default)">
+  <Profiles>
+    <Profile Name="(Default)" />
+  </Profiles>
+  <Settings />
+</SettingsFile>
\ No newline at end of file
diff --git a/Engine/Source/Programs/Timing Data Investigator/Timing Data Investigator/Timing Data Investigator.csproj b/Engine/Source/Programs/Timing Data Investigator/Timing Data Investigator/Timing Data Investigator.csproj
new file mode 100644
--- /dev/null
+++ b/Engine/Source/Programs/Timing Data Investigator/Timing Data Investigator/Timing Data Investigator.csproj	
@@ -0,0 +1,131 @@
+﻿<?xml version="1.0" encoding="utf-8"?>
+<Project ToolsVersion="15.0" xmlns="http://schemas.microsoft.com/developer/msbuild/2003">
+  <Import Project="$(MSBuildExtensionsPath)\$(MSBuildToolsVersion)\Microsoft.Common.props" Condition="Exists('$(MSBuildExtensionsPath)\$(MSBuildToolsVersion)\Microsoft.Common.props')" />
+  <PropertyGroup>
+    <Configuration Condition=" '$(Configuration)' == '' ">Debug</Configuration>
+    <Platform Condition=" '$(Platform)' == '' ">AnyCPU</Platform>
+    <ProjectGuid>{3F3BEAB2-B096-4420-9E90-BD0CB55867CE}</ProjectGuid>
+    <OutputType>WinExe</OutputType>
+    <RootNamespace>Timing_Data_Investigator</RootNamespace>
+    <AssemblyName>Timing Data Investigator</AssemblyName>
+    <TargetFrameworkVersion>v4.6.2</TargetFrameworkVersion>
+    <FileAlignment>512</FileAlignment>
+    <ProjectTypeGuids>{60dc8134-eba5-43b8-bcc9-bb4bc16c2548};{FAE04EC0-301F-11D3-BF4B-00C04F79EFBC}</ProjectTypeGuids>
+    <WarningLevel>4</WarningLevel>
+    <AutoGenerateBindingRedirects>true</AutoGenerateBindingRedirects>
+    <Deterministic>true</Deterministic>
+    <TargetFrameworkProfile />
+  </PropertyGroup>
+  <PropertyGroup Condition=" '$(Configuration)|$(Platform)' == 'Debug|AnyCPU' ">
+    <PlatformTarget>x64</PlatformTarget>
+    <DebugSymbols>true</DebugSymbols>
+    <DebugType>full</DebugType>
+    <Optimize>false</Optimize>
+    <OutputPath>bin\Debug\</OutputPath>
+    <DefineConstants>DEBUG;TRACE</DefineConstants>
+    <ErrorReport>prompt</ErrorReport>
+    <WarningLevel>4</WarningLevel>
+  </PropertyGroup>
+  <PropertyGroup Condition=" '$(Configuration)|$(Platform)' == 'Release|AnyCPU' ">
+    <PlatformTarget>AnyCPU</PlatformTarget>
+    <DebugType>pdbonly</DebugType>
+    <Optimize>true</Optimize>
+    <OutputPath>bin\Release\</OutputPath>
+    <DefineConstants>TRACE</DefineConstants>
+    <ErrorReport>prompt</ErrorReport>
+    <WarningLevel>4</WarningLevel>
+  </PropertyGroup>
+  <ItemGroup>
+    <Reference Include="System" />
+    <Reference Include="System.Data" />
+    <Reference Include="System.Xml" />
+    <Reference Include="Microsoft.CSharp" />
+    <Reference Include="System.Core" />
+    <Reference Include="System.Xml.Linq" />
+    <Reference Include="System.Data.DataSetExtensions" />
+    <Reference Include="System.Net.Http" />
+    <Reference Include="System.Xaml">
+      <RequiredTargetFramework>4.0</RequiredTargetFramework>
+    </Reference>
+    <Reference Include="WindowsBase" />
+    <Reference Include="PresentationCore" />
+    <Reference Include="PresentationFramework" />
+  </ItemGroup>
+  <ItemGroup>
+    <ApplicationDefinition Include="App.xaml">
+      <Generator>MSBuild:Compile</Generator>
+      <SubType>Designer</SubType>
+    </ApplicationDefinition>
+    <Compile Include="Commands\OpenTimingDataCommand.cs" />
+    <Compile Include="Controls\TimingDataTab.xaml.cs">
+      <DependentUpon>TimingDataTab.xaml</DependentUpon>
+    </Compile>
+    <Compile Include="TreeGridElement.cs" />
+    <Compile Include="TreeGridElementComparer.cs" />
+    <Compile Include="TreeGridFlatModel.cs" />
+    <Compile Include="TreeGridModel.cs" />
+    <Page Include="Controls\TimingDataGrid.xaml">
+      <SubType>Designer</SubType>
+      <Generator>MSBuild:Compile</Generator>
+    </Page>
+    <Page Include="Controls\TimingDataTab.xaml">
+      <SubType>Designer</SubType>
+      <Generator>MSBuild:Compile</Generator>
+    </Page>
+    <Page Include="MainWindow.xaml">
+      <Generator>MSBuild:Compile</Generator>
+      <SubType>Designer</SubType>
+    </Page>
+    <Compile Include="App.xaml.cs">
+      <DependentUpon>App.xaml</DependentUpon>
+      <SubType>Code</SubType>
+    </Compile>
+    <Compile Include="Controls\TimingDataGrid.xaml.cs">
+      <DependentUpon>TimingDataGrid.xaml</DependentUpon>
+    </Compile>
+    <Compile Include="Converters.cs" />
+    <Compile Include="MainWindow.xaml.cs">
+      <DependentUpon>MainWindow.xaml</DependentUpon>
+      <SubType>Code</SubType>
+    </Compile>
+  </ItemGroup>
+  <ItemGroup>
+    <Compile Include="Models\TimingDataSummaryViewModel.cs" />
+    <Compile Include="Models\TimingDataViewModel.cs" />
+    <Compile Include="Properties\AssemblyInfo.cs">
+      <SubType>Code</SubType>
+    </Compile>
+    <Compile Include="Properties\Resources.Designer.cs">
+      <AutoGen>True</AutoGen>
+      <DesignTime>True</DesignTime>
+      <DependentUpon>Resources.resx</DependentUpon>
+    </Compile>
+    <Compile Include="Properties\Settings.Designer.cs">
+      <AutoGen>True</AutoGen>
+      <DependentUpon>Settings.settings</DependentUpon>
+      <DesignTimeSharedInput>True</DesignTimeSharedInput>
+    </Compile>
+    <EmbeddedResource Include="Properties\Resources.resx">
+      <Generator>ResXFileCodeGenerator</Generator>
+      <LastGenOutput>Resources.Designer.cs</LastGenOutput>
+    </EmbeddedResource>
+    <None Include="Properties\Settings.settings">
+      <Generator>SettingsSingleFileGenerator</Generator>
+      <LastGenOutput>Settings.Designer.cs</LastGenOutput>
+    </None>
+  </ItemGroup>
+  <ItemGroup>
+    <None Include="App.config" />
+  </ItemGroup>
+  <ItemGroup>
+    <ProjectReference Include="..\..\UnrealBuildTool\UnrealBuildTool.csproj">
+      <Project>{fd7c5e1a-cfe4-4fd5-a525-1eb1599a39ac}</Project>
+      <Name>UnrealBuildTool</Name>
+    </ProjectReference>
+    <ProjectReference Include="..\..\DotNETCommon\DotNETUtilities\DotNETUtilities.csproj">
+      <Project>{5d7d66e8-8c76-4af9-b3ec-2ef03421d730}</Project>
+      <Name>DotNETUtilities</Name>
+    </ProjectReference>
+  </ItemGroup>
+  <Import Project="$(MSBuildToolsPath)\Microsoft.CSharp.targets" />
+</Project>
\ No newline at end of file
diff --git a/Engine/Source/Programs/Timing Data Investigator/Timing Data Investigator/TreeGridElement.cs b/Engine/Source/Programs/Timing Data Investigator/Timing Data Investigator/TreeGridElement.cs
new file mode 100644
--- /dev/null
+++ b/Engine/Source/Programs/Timing Data Investigator/Timing Data Investigator/TreeGridElement.cs	
@@ -0,0 +1,308 @@
+﻿using System;
+using System.Collections;
+using System.Collections.ObjectModel;
+using System.Collections.Specialized;
+using System.Windows;
+
+namespace Timing_Data_Investigator
+{
+	public class TreeGridElement : ContentElement
+	{
+		private const string NullItemError = "The item added to the collection cannot be null.";
+
+		public static readonly RoutedEvent        ExpandingEvent;
+		public static readonly RoutedEvent        ExpandedEvent;
+		public static readonly RoutedEvent        CollapsingEvent;
+		public static readonly RoutedEvent        CollapsedEvent;
+		public static readonly DependencyProperty HasChildrenProperty;
+		public static readonly DependencyProperty IsExpandedProperty;
+		public static readonly DependencyProperty LevelProperty;
+
+		public TreeGridElement Parent { get; set; }
+		public TreeGridModel Model { get; set; }
+		public ObservableCollection<TreeGridElement> Children { get; set; }
+
+		static TreeGridElement()
+		{
+			// Register dependency properties
+			HasChildrenProperty = RegisterHasChildrenProperty();
+			IsExpandedProperty  = RegisterIsExpandedProperty();
+			LevelProperty       = RegisterLevelProperty();
+
+			// Register routed events
+			ExpandingEvent  = RegisterExpandingEvent();
+			ExpandedEvent   = RegisterExpandedEvent();
+			CollapsingEvent = RegisterCollapsingEvent();
+			CollapsedEvent  = RegisterCollapsedEvent();
+		}
+
+		public TreeGridElement()
+		{
+			// Initialize the element
+			Children = new ObservableCollection<TreeGridElement>();
+
+			// Attach events
+			Children.CollectionChanged += OnChildrenChanged;
+		}
+
+		internal void SetModel(TreeGridModel NewModel, TreeGridElement NewParent = null)
+		{
+			// Set the element information
+			Model = NewModel;
+			Parent = NewParent;
+			Level = (NewParent != null) ? NewParent.Level + 1 : 0;
+
+			// Iterate through all child elements
+			foreach (TreeGridElement Child in Children)
+			{
+				// Set the model for the child
+				Child.SetModel(NewModel, this);
+			}
+		}
+
+		protected virtual void OnExpanding() { }
+		protected virtual void OnExpanded() { }
+		protected virtual void OnCollapsing() { }
+		protected virtual void OnCollapsed() { }
+
+		private void OnChildrenChanged(object sender, NotifyCollectionChangedEventArgs args)
+		{
+			// Process the event
+			switch(args.Action)
+			{
+				case NotifyCollectionChangedAction.Add:
+
+					// Process added child
+					OnChildAdded(args.NewItems[0]);
+					break;
+
+				case NotifyCollectionChangedAction.Replace:
+
+					// Process replaced child
+					OnChildReplaced((TreeGridElement)args.OldItems[0], args.NewItems[0], args.NewStartingIndex);
+					break;
+
+				case NotifyCollectionChangedAction.Remove:
+
+					// Process removed child
+					OnChildRemoved((TreeGridElement)args.OldItems[0]);
+					break;
+
+				case NotifyCollectionChangedAction.Reset:
+
+					// Process cleared children
+					OnChildrenCleared(args.OldItems);
+					break;
+			}
+		}
+
+		private void OnChildAdded(object item)
+		{
+			// Verify the new child
+			TreeGridElement child = VerifyItem(item);
+
+			// Set the model for the child
+			child.SetModel(Model, this);
+
+			// Notify the model that a child was added to the item
+			Model?.OnChildAdded(child);
+		}
+
+		private void OnChildReplaced(TreeGridElement oldChild, object item, int index)
+		{
+			// Verify the new child
+			TreeGridElement child = VerifyItem(item);
+
+			// Clear the model for the old child
+			oldChild.SetModel(null);
+
+			// Notify the model that a child was replaced
+			Model?.OnChildReplaced(oldChild, child, index);
+		}
+
+		private void OnChildRemoved(TreeGridElement child)
+		{
+			// Clear the model for the child
+			child.SetModel(null);
+
+			// Notify the model that a child was removed from the item
+			Model?.OnChildRemoved(child);
+		}
+
+		private void OnChildrenCleared(IList children)
+		{
+            if (children == null)
+            {
+                return;
+            }
+
+			// Iterate through all of the children
+			foreach (TreeGridElement child in children)
+			{
+				// Clear the model for the child
+				child.SetModel(null);
+			}
+
+			// Notify the model that all of the children were removed from the item
+			Model?.OnChildrenRemoved(this, children);
+		}
+
+		internal static TreeGridElement VerifyItem(object Item)
+		{
+			// Is the item valid?
+			if (Item == null)
+			{
+				// The item is not valid
+				throw new ArgumentNullException(nameof(Item), NullItemError);
+			}
+
+			// Return the element
+			return (TreeGridElement)Item;
+		}
+
+		private static void OnIsExpandedChanged(DependencyObject element, DependencyPropertyChangedEventArgs args)
+		{
+			// Get the tree item
+			TreeGridElement item = (TreeGridElement)element;
+
+			// Is the item being expanded?
+			if ((bool)args.NewValue)
+			{
+				// Raise expanding event
+				item.RaiseEvent(new RoutedEventArgs(ExpandingEvent, item));
+
+				// Execute derived expanding handler
+				item.OnExpanding();
+
+				// Expand the item in the model
+				item.Model?.Expand(item);
+
+				// Raise expanded event
+				item.RaiseEvent(new RoutedEventArgs(ExpandedEvent, item));
+
+				// Execute derived expanded handler
+				item.OnExpanded();
+			}
+			else
+			{
+				// Raise collapsing event
+				item.RaiseEvent(new RoutedEventArgs(CollapsingEvent, item));
+
+				// Execute derived collapsing handler
+				item.OnCollapsing();
+
+				// Collapse the item in the model
+				item.Model?.Collapse(item);
+
+				// Raise collapsed event
+				item.RaiseEvent(new RoutedEventArgs(CollapsedEvent, item));
+
+				// Execute derived collapsed handler
+				item.OnCollapsed();
+			}
+		}
+
+		private static DependencyProperty RegisterHasChildrenProperty()
+		{
+			// Create the property metadata
+			FrameworkPropertyMetadata metadata = new FrameworkPropertyMetadata()
+			{
+				DefaultValue = false
+			};
+
+			// Register the property
+			return DependencyProperty.Register(nameof(HasChildren), typeof(bool), typeof(TreeGridElement), metadata);
+		}
+
+		private static DependencyProperty RegisterIsExpandedProperty()
+		{
+			// Create the property metadata
+			FrameworkPropertyMetadata metadata = new FrameworkPropertyMetadata()
+			{
+				DefaultValue            = false,
+				PropertyChangedCallback = OnIsExpandedChanged
+			};
+
+			// Register the property
+			return DependencyProperty.Register(nameof(IsExpanded), typeof(bool), typeof(TreeGridElement), metadata);
+		}
+
+		private static DependencyProperty RegisterLevelProperty()
+		{
+			// Create the property metadata
+			FrameworkPropertyMetadata metadata = new FrameworkPropertyMetadata()
+			{
+				DefaultValue = 0
+			};
+
+			// Register the property
+			return DependencyProperty.Register(nameof(Level), typeof(int), typeof(TreeGridElement), metadata);
+		}
+
+		private static RoutedEvent RegisterExpandingEvent()
+		{
+			// Register the event
+			return EventManager.RegisterRoutedEvent(nameof(Expanding), RoutingStrategy.Bubble, typeof(RoutedEventHandler), typeof(TreeGridElement));
+		}
+
+		private static RoutedEvent RegisterExpandedEvent()
+		{
+			// Register the event
+			return EventManager.RegisterRoutedEvent(nameof(Expanded), RoutingStrategy.Bubble, typeof(RoutedEventHandler), typeof(TreeGridElement));
+		}
+
+		private static RoutedEvent RegisterCollapsingEvent()
+		{
+			// Register the event
+			return EventManager.RegisterRoutedEvent(nameof(Collapsing), RoutingStrategy.Bubble, typeof(RoutedEventHandler), typeof(TreeGridElement));
+		}
+
+		private static RoutedEvent RegisterCollapsedEvent()
+		{
+			// Register the event
+			return EventManager.RegisterRoutedEvent(nameof(Collapsed), RoutingStrategy.Bubble, typeof(RoutedEventHandler), typeof(TreeGridElement));
+		}
+
+		public bool HasChildren
+		{
+			get { return (bool)GetValue(HasChildrenProperty); }
+			set { SetValue(HasChildrenProperty, value); }
+		}
+
+		public bool IsExpanded
+		{
+			get { return (bool)GetValue(IsExpandedProperty); }
+			set { SetValue(IsExpandedProperty, value); }
+		}
+
+		public int Level
+		{
+			get { return (int)GetValue(LevelProperty); }
+			set { SetValue(LevelProperty, value); }
+		}
+
+		public event RoutedEventHandler Expanding
+		{
+			add { AddHandler(ExpandingEvent, value); }
+			remove { RemoveHandler(ExpandingEvent, value); }
+		}
+
+		public event RoutedEventHandler Expanded
+		{
+			add { AddHandler(ExpandedEvent, value); }
+			remove { RemoveHandler(ExpandedEvent, value); }
+		}
+
+		public event RoutedEventHandler Collapsing
+		{
+			add { AddHandler(CollapsingEvent, value); }
+			remove { RemoveHandler(CollapsingEvent, value); }
+		}
+
+		public event RoutedEventHandler Collapsed
+		{
+			add { AddHandler(CollapsedEvent, value); }
+			remove { RemoveHandler(CollapsedEvent, value); }
+		}
+	}
+}
\ No newline at end of file
diff --git a/Engine/Source/Programs/Timing Data Investigator/Timing Data Investigator/TreeGridElementComparer.cs b/Engine/Source/Programs/Timing Data Investigator/Timing Data Investigator/TreeGridElementComparer.cs
new file mode 100644
--- /dev/null
+++ b/Engine/Source/Programs/Timing Data Investigator/Timing Data Investigator/TreeGridElementComparer.cs	
@@ -0,0 +1,42 @@
+﻿using System;
+using System.Collections.Generic;
+using System.ComponentModel;
+using System.Linq;
+
+namespace Timing_Data_Investigator
+{
+    public class PropertyComparer<TClassType> : IComparer<TClassType>
+    {
+        private string sCompareProperty;
+        private bool bSortDescending;
+
+        public PropertyComparer(string CompareProperty, bool SortDescending)
+        {
+            sCompareProperty = CompareProperty;
+            bSortDescending = SortDescending;
+        }
+
+        public int Compare(TClassType x, TClassType y)
+        {
+            var ComparePropertyInfo = typeof(TClassType).GetProperties().FirstOrDefault(p => p.Name == sCompareProperty);
+            if (ComparePropertyInfo == null)
+            {
+                throw new InvalidOperationException($"Class '{nameof(TClassType)}' does not contain a property named '{sCompareProperty}'!");
+            }
+
+            if (!ComparePropertyInfo.PropertyType.GetInterfaces().Any(i => i.Name.Contains("IComparable")))
+            {
+                throw new InvalidOperationException($"Property type '{ComparePropertyInfo.PropertyType.Name}' is not an IComparable!");
+            }
+
+            var xPropValue = ComparePropertyInfo.GetValue(x) as IComparable;
+            var yPropValue = ComparePropertyInfo.GetValue(y) as IComparable;
+            if (bSortDescending)
+            {
+                return yPropValue.CompareTo(xPropValue);
+            }
+
+            return xPropValue.CompareTo(yPropValue);
+        }
+    }
+}
diff --git a/Engine/Source/Programs/Timing Data Investigator/Timing Data Investigator/TreeGridFlatModel.cs b/Engine/Source/Programs/Timing Data Investigator/Timing Data Investigator/TreeGridFlatModel.cs
new file mode 100644
--- /dev/null
+++ b/Engine/Source/Programs/Timing Data Investigator/Timing Data Investigator/TreeGridFlatModel.cs	
@@ -0,0 +1,106 @@
+﻿using System;
+using System.Collections.Generic;
+using System.Collections.ObjectModel;
+using System.Collections.Specialized;
+
+namespace Timing_Data_Investigator
+{
+	public class TreeGridFlatModel : ObservableCollection<TreeGridElement>
+	{
+		private const string ModificationError = "The collection cannot be modified by the user.";
+
+		private bool modification;
+		private HashSet<TreeGridElement> keys;
+
+		public TreeGridFlatModel()
+		{
+			// Initialize the model
+			keys = new HashSet<TreeGridElement>();
+		}
+
+		internal bool ContainsKey(TreeGridElement item)
+		{
+			// Return a value indicating if the item is within the model
+			return keys.Contains(item);
+		}
+
+		internal void PrivateInsert(int index, TreeGridElement item)
+		{
+			// Set the modification flag
+			modification = true;
+
+			// Add the item to the model
+			Insert(index, item);
+
+			// Add the item to the keys
+			keys.Add(item);
+
+			// Clear the modification flag
+			modification = false;
+		}
+
+		internal void PrivateInsertRange(int index, IList<TreeGridElement> items)
+		{
+			// Set the modification flag
+			modification = true;
+
+			// Iterate through all of the children within the items
+			foreach (TreeGridElement child in items)
+			{
+				// Add the child to the model
+				Insert(index++, child);
+
+				// Add the child to the keys
+				keys.Add(child);
+			}
+
+			// Clear the modification flag
+			modification = false;
+		}
+
+		internal void PrivateRemoveRange(int index, int count)
+		{
+			// Set the modification flag
+			modification = true;
+
+			// Iterate through all of the items to remove from the model
+			for (int itemIndex = 0; itemIndex < count; itemIndex++)
+			{
+				// Remove the item from the keys
+				keys.Remove(Items[index]);
+
+				// Remove the item from the model
+				RemoveAt(index);
+			}
+
+			// Clear the modification flag
+			modification = false;
+		}
+
+		internal void PrivateClear()
+		{
+			// Set the modification flag
+			modification = true;
+
+			// Clear keys and items.
+			keys.Clear();
+			Clear();
+
+			// Clear the modification flag
+			modification = false;
+		}
+
+		protected override void OnCollectionChanged(NotifyCollectionChangedEventArgs args)
+		{
+			// Is the modification flag set?
+			if (!modification)
+			{
+				// The collection is for internal use only
+				throw new InvalidOperationException(ModificationError);
+			}
+
+			// Call base method
+			base.OnCollectionChanged(args);
+		}
+	}
+}
\ No newline at end of file
diff --git a/Engine/Source/Programs/Timing Data Investigator/Timing Data Investigator/TreeGridModel.cs b/Engine/Source/Programs/Timing Data Investigator/Timing Data Investigator/TreeGridModel.cs
new file mode 100644
--- /dev/null
+++ b/Engine/Source/Programs/Timing Data Investigator/Timing Data Investigator/TreeGridModel.cs	
@@ -0,0 +1,193 @@
+﻿using System.Collections;
+using System.Collections.Generic;
+using System.Collections.ObjectModel;
+using System.Collections.Specialized;
+
+namespace Timing_Data_Investigator
+{
+	public class TreeGridModel : ObservableCollection<TreeGridElement>
+	{
+		private List<TreeGridElement> ItemCache;
+
+		public TreeGridModel()
+		{
+			// Initialize the model
+			ItemCache = new List<TreeGridElement>();
+			FlatModel = new TreeGridFlatModel();
+		}
+
+		public TreeGridFlatModel FlatModel { get; private set; }
+
+		protected override void OnCollectionChanged(NotifyCollectionChangedEventArgs args)
+		{
+			// Process the event
+			switch (args.Action)
+			{
+				case NotifyCollectionChangedAction.Add:
+
+					// Process added item
+					OnRootAdded(args.NewItems[0]);
+					break;
+			}
+		}
+
+		internal void Expand(TreeGridElement item)
+		{
+			// Do we need to expand the item?
+			if (!FlatModel.ContainsKey(item) || !item.IsExpanded)
+			{
+				// We do not need to expand the item
+				return;
+			}
+
+			// Clear the item cache
+			ItemCache.Clear();
+
+			// Cache the flat children for the item
+			CacheFlatChildren(item);
+
+			// Get the insertion index for the item
+			int index = (FlatModel.IndexOf(item) + 1);
+
+			// Add the flat children to the flat model
+			FlatModel.PrivateInsertRange(index, ItemCache);
+		}
+
+		internal void Collapse(TreeGridElement item)
+		{
+			// Do we need to collapse the item?
+			if (!FlatModel.ContainsKey(item))
+			{
+				// We do not need to collapse the item
+				return;
+			}
+
+			// Get the collapse information
+			int index = (FlatModel.IndexOf(item) + 1);
+			int count = CountFlatChildren(item);
+
+			// Remove the items from the flat model to collapse them
+			FlatModel.PrivateRemoveRange(index, count);
+		}
+
+		internal void OnChildAdded(TreeGridElement child)
+		{
+			// Get the parent of the child
+			TreeGridElement parent = child.Parent;
+
+			// Check if the parent is expanded
+			if (!FlatModel.ContainsKey(parent) || !parent.IsExpanded)
+			{
+				// We don't need to update the flat model
+				return;
+			}
+
+			// Find the insertion index for the child into the flat model
+			int index = FindFlatInsertionIndex(child);
+
+			// Insert the child into the flat model
+			FlatModel.PrivateInsert(index, child);
+
+			// Expand the child
+			Expand(child);
+		}
+
+		internal void OnChildReplaced(TreeGridElement oldChild, TreeGridElement child, int index)
+		{
+		}
+
+		internal void OnChildRemoved(TreeGridElement child)
+		{
+		}
+
+		internal void OnChildrenRemoved(TreeGridElement parent, IList children)
+		{
+		}
+
+		private void OnRootAdded(object Item)
+		{
+			// Verify the root item
+			TreeGridElement Root = TreeGridElement.VerifyItem(Item);
+
+			// Set the model for the root
+			Root.SetModel(this);
+
+			// Find the index for insertion into the flat model
+			int index = FindFlatInsertionIndex(Root);
+
+			// Insert the root into the flat model
+			FlatModel.PrivateInsert(index, Root);
+
+			// Expand the root
+			Expand(Root);
+		}
+
+		private void CacheFlatChildren(TreeGridElement item)
+		{
+			// Iterate through all of the children within the item
+			foreach (TreeGridElement child in item.Children)
+			{
+				// Add the child to the item cache
+				ItemCache.Add(child);
+
+				// Is the child expanded?
+				if (child.IsExpanded)
+				{
+					// Recursively cache the children within the child
+					CacheFlatChildren(child);
+				}
+			}
+		}
+
+		private int CountFlatChildren(TreeGridElement item)
+		{
+			// Initialize child count
+			int children = item.Children.Count;
+
+			// Iterate through each child
+			foreach (TreeGridElement child in item.Children)
+			{
+				// Is the child expanded?
+				if (child.IsExpanded)
+				{
+					// Recursively count the children
+					children += CountFlatChildren(child);
+				}
+			}
+
+			// Return the number of flat children
+			return children;
+		}
+
+		private int FindFlatInsertionIndex(TreeGridElement item)
+		{
+			// Get the search information
+			TreeGridElement parent = item.Parent;
+			IList<TreeGridElement> items = ((parent != null) ? parent.Children : this);
+			int index = items.IndexOf(item);
+			int lastIndex = (items.Count - 1);
+
+			// Determine if the item is the last item in the items
+			if (index < lastIndex)
+			{
+				// Return the insertion index using the items
+				return FlatModel.IndexOf(items[(index + 1)]);
+			}
+
+			// Is the parent valid?
+			else if (parent != null)
+			{
+				// Determine the number of flat children the parent has
+				int children = CountFlatChildren(parent);
+
+				// Return the insertion index using the number of flat children
+				return (FlatModel.IndexOf(parent) + children);
+			}
+			else
+			{
+				// Return the flat model count
+				return FlatModel.Count;
+			}
+		}
+	}
+}
\ No newline at end of file
diff --git a/Engine/Source/Programs/UnrealBuildTool/System/AggregateParsedTimingInfo.cs b/Engine/Source/Programs/UnrealBuildTool/System/AggregateParsedTimingInfo.cs
--- a/Engine/Source/Programs/UnrealBuildTool/System/AggregateParsedTimingInfo.cs
+++ b/Engine/Source/Programs/UnrealBuildTool/System/AggregateParsedTimingInfo.cs
@@ -22,6 +22,9 @@ namespace UnrealBuildTool
 		private TimingData FileTimingData;
 		private ConcurrentDictionary<string, int> DecompressedFileSizes = new ConcurrentDictionary<string, int>();
 		private ConcurrentDictionary<string, byte[]> CompressedFiles = new ConcurrentDictionary<string, byte[]>();
+		private ConcurrentBag<TimingData> AggregateIncludes = new ConcurrentBag<TimingData>();
+		private ConcurrentBag<TimingData> AggregateClasses = new ConcurrentBag<TimingData>();
+		private ConcurrentBag<TimingData> AggregateFunctions = new ConcurrentBag<TimingData>();
 		private object AddFileLock = new object();
 
 		public override int Execute(CommandLineArguments Arguments)
@@ -38,6 +41,11 @@ namespace UnrealBuildTool
 			TimingData AggregateData = new TimingData() { Name = AggregateName, Type = TimingDataType.Aggregate };
 			AggregateData.AddChild(FileTimingData);
 
+			// Group the includes, classes, and functions by name and sum them up then add to the aggregate.
+			GroupTimingDataOnName(AggregateData, "Include Timings", AggregateIncludes);//, 100);
+			GroupTimingDataOnName(AggregateData, "Class Timings", AggregateClasses);//, 100);
+			GroupTimingDataOnName(AggregateData, "Function Timings", AggregateFunctions);//, 100);
+
 			// Write out aggregate summary.
 			string OutputFile = Path.Combine(ManifestFile.Directory.FullName, String.Format("{0}.timing.bin", AggregateName));
 			using (BinaryWriter Writer = new BinaryWriter(File.Open(OutputFile, FileMode.Create)))
@@ -67,6 +75,20 @@ namespace UnrealBuildTool
 			return 0;
 		}
 
+		private void GroupTimingDataOnName(TimingData AggregateData, string TimingName, IEnumerable<TimingData> UngroupedData, int NumberToKeep = int.MaxValue)
+		{
+			string GroupName = String.Format("Top {0} {1}", NumberToKeep, TimingName);
+			TimingData GroupedTimingData = new TimingData() { Name = GroupName, Type = TimingDataType.Summary };
+			IEnumerable<IGrouping<string, TimingData>> Groups = UngroupedData.GroupBy(i => i.Name).OrderByDescending(g => g.Sum(d => d.ExclusiveDuration)).ToList();
+			foreach (IGrouping<string, TimingData> Group in Groups.Take(NumberToKeep))
+			{
+				TimingData GroupedData = new TimingData() { Name = Group.Key, ExclusiveDuration = Group.Sum(d => d.ExclusiveDuration) };
+				GroupedTimingData.Children.Add(Group.Key, GroupedData);
+			}
+
+			AggregateData.AddChild(GroupedTimingData);
+		}
+
 		private void ParseTimingDataFile(string ParsedFileName)
 		{
 			// Convert input file back into summary objects.
@@ -95,7 +117,16 @@ namespace UnrealBuildTool
 				TimingData SummarizedIncludes = new TimingData() { Parent = SummarizedTimingData, Type = TimingDataType.Summary };
 				Task QueueIncludesTask = Task.Run(() =>
 				{
-					SummarizeTimingData(ParsedTimingData.Children["IncludeTimings"].Children.Values, SummarizedIncludes, "Includes", 10);
+					IEnumerable<TimingData> Includes = ParsedTimingData.Children["IncludeTimings"].Children.Values;
+					SummarizeTimingData(Includes, SummarizedIncludes, "Includes", 10);
+
+					// Flatten the includes for aggregation.
+					Dictionary<string, TimingData> FlattenedIncludes = new Dictionary<string, TimingData>();
+					FlattenIncludes(FlattenedIncludes, Includes);
+					foreach (TimingData Include in Includes)
+					{
+						AggregateIncludes.Add(Include);
+					}
 				});
 
 				TimingData SummarizedClasses = new TimingData() { Parent = SummarizedTimingData, Type = TimingDataType.Summary };
@@ -104,6 +135,10 @@ namespace UnrealBuildTool
 					// Collapse templates into single entries.
 					IEnumerable<TimingData> CollapsedClasses = GroupChildren(ParsedTimingData.Children["ClassTimings"].Children.Values, TimingDataType.Class);
 					SummarizeTimingData(CollapsedClasses, SummarizedClasses, "Classes", 10);
+					foreach (TimingData Class in CollapsedClasses)
+					{
+						AggregateClasses.Add(new TimingData() { Name = Class.Name, Type = TimingDataType.Class, ExclusiveDuration = Class.InclusiveDuration });
+					}
 				});
 
 				TimingData SummarizedFunctions = new TimingData() { Parent = SummarizedTimingData, Type = TimingDataType.Summary };
@@ -112,6 +147,10 @@ namespace UnrealBuildTool
 					// Collapse templates into single entries.
 					IEnumerable<TimingData> CollapsedFunctions = GroupChildren(ParsedTimingData.Children["FunctionTimings"].Children.Values, TimingDataType.Function);
 					SummarizeTimingData(CollapsedFunctions, SummarizedFunctions, "Functions", 10);
+					foreach (TimingData Function in CollapsedFunctions)
+					{
+						AggregateFunctions.Add(new TimingData() { Name = Function.Name, Type = TimingDataType.Function, ExclusiveDuration = Function.InclusiveDuration });
+					}
 				});
 
 				while (!CompressDataTask.IsCompleted || !QueueIncludesTask.IsCompleted || !QueueClassesTask.IsCompleted || !QueueFunctionsTask.IsCompleted)
@@ -130,6 +169,24 @@ namespace UnrealBuildTool
 			}
 		}
 
+		private void FlattenIncludes(Dictionary<string, TimingData> Includes, IEnumerable<TimingData> IncludeData)
+		{
+			foreach (TimingData Include in IncludeData)
+			{
+				TimingData AggregatedInclude;
+				if (Includes.TryGetValue(Include.Name, out AggregatedInclude))
+				{
+					AggregatedInclude.ExclusiveDuration += Include.ExclusiveDuration;
+				}
+				else
+				{
+					Includes.Add(Include.Name, new TimingData() { Name = Include.Name, Type = TimingDataType.Include, ExclusiveDuration = Include.ExclusiveDuration });
+				}
+
+				FlattenIncludes(Includes, Include.Children.Values);
+			}
+		}
+
 		private IEnumerable<TimingData> GroupChildren(IEnumerable<TimingData> Children, TimingDataType Type)
 		{
 			List<TimingData> GroupedChildren = new List<TimingData>();
